PerfCounters for Wash.

- Configured by GameDir/Config/PerfCounters.json (should be deployed by UAT now).
- Shared memory regions have a name of perfcounters-of-pid-%d.

diff --git a/Engine/Source/Programs/AutomationTool/Linux/LinuxPlatform.Automation.cs b/Engine/Source/Programs/AutomationTool/Linux/LinuxPlatform.Automation.cs
--- a/Engine/Source/Programs/AutomationTool/Linux/LinuxPlatform.Automation.cs
+++ b/Engine/Source/Programs/AutomationTool/Linux/LinuxPlatform.Automation.cs
@@ -45,6 +45,8 @@ public abstract class BaseLinuxPlatform : Platform
 			}
 
 			SC.StageFiles(StagedFileType.NonUFS, CombinePaths(SC.LocalRoot, "Engine/Config"), "controller.vdf", false, null, CommandUtils.CombinePaths(SC.RelativeProjectRootForStage, "Saved/Config"));
+			// copy optional perfcounters definition file
+			SC.StageFiles(StagedFileType.NonUFS, CombinePaths(SC.ProjectRoot, "Config"), "PerfCounters.json", false, null, CommandUtils.CombinePaths(SC.RelativeProjectRootForStage, "Saved/Config"), true);
 		}
 
 		// assume that we always have to deploy OpenAL (FIXME: should be automatic)
diff --git a/Engine/Source/Runtime/IPC/Private/IPCModule.cpp b/Engine/Source/Runtime/IPC/Private/IPCModule.cpp
new file mode 100644
--- /dev/null
+++ b/Engine/Source/Runtime/IPC/Private/IPCModule.cpp
@@ -0,0 +1,7 @@
+#include "IPCModule.h"
+
+class FIPCModule : public IIPCModule
+{
+};
+
+IMPLEMENT_MODULE(FIPCModule, IPC)
diff --git a/Engine/Source/Runtime/IPC/Public/IPCModule.h b/Engine/Source/Runtime/IPC/Public/IPCModule.h
new file mode 100644
--- /dev/null
+++ b/Engine/Source/Runtime/IPC/Public/IPCModule.h
@@ -0,0 +1,37 @@
+// Copyright 1998-2014 Epic Games, Inc. All Rights Reserved.
+
+#pragma once
+
+#include "ModuleManager.h"
+
+
+/**
+* The public interface to this module
+*/
+class IIPCModule : public IModuleInterface
+{
+
+public:
+
+	/**
+	 * Singleton-like access to this module's interface.  This is just for convenience!
+	 * Beware of calling this during the shutdown phase, though.  Your module might have been unloaded already.
+	 *
+	 * @return Returns singleton instance, loading the module on demand if needed
+	 */
+	static inline IIPCModule& Get()
+	{
+		return FModuleManager::LoadModuleChecked< IIPCModule >("IPC");
+	}
+
+	/**
+	 * Checks to see if this module is loaded and ready.  It is only valid to call Get() if IsAvailable() returns true.
+	 *
+	 * @return True if the module is loaded and ready to use
+	 */
+	static inline bool IsAvailable()
+	{
+		return FModuleManager::Get().IsModuleLoaded("IPC");
+	}
+};
+
